Fixing misc issues with default AEC template:

asset preview no longer have bloom and uses proper environment map from the editor content
leveltempates moved out of the project
ini settings tweaked

diff --git a/Templates/TP_AEC_BlankBP/Config/DefaultEngine.ini b/Templates/TP_AEC_BlankBP/Config/DefaultEngine.ini
--- a/Templates/TP_AEC_BlankBP/Config/DefaultEngine.ini
+++ b/Templates/TP_AEC_BlankBP/Config/DefaultEngine.ini
@@ -37,8 +37,6 @@ EditorStartupMap=/Game/Main.Main
 [/Script/UnrealEd.UnrealEdEngine]
 ; Remove Engine Template maps
 -TemplateMapInfos=(ThumbnailTexture=Texture2D'/Engine/Maps/Templates/Thumbnails/Default.Default',Map="/Engine/Maps/Templates/Template_Default")
+-TemplateMapInfos=(ThumbnailTexture=Texture2D'/Engine/Maps/Templates/Thumbnails/TimeOfDay.TimeOfDay',Map="/Engine/Maps/Templates/TimeOfDay_Default")
 -TemplateMapInfos=(ThumbnailTexture=Texture2D'/Engine/Maps/Templates/Thumbnails/VR-Basic.VR-Basic',Map="/Engine/Maps/Templates/VR-Basic")
-
-; Add Predefined Template Maps that are suitable for this template
-;+TemplateMapInfos=(ThumbnailTexture=Texture2D'/Game/Developers/Maps/Thumbnails/AssetName.AssetName',Map="/Game/Maps/Templates/MapName")
-+TemplateMapInfos=(ThumbnailTexture=Texture2D'/Game/Developers/LevelTemplates/Thumbnails/Default.Default',Map="/Game/Developers/LevelTemplates/Default")
\ No newline at end of file
++TemplateMapInfos=(ThumbnailTexture=Texture2D'/Engine/Maps/Templates/Thumbnails/Basic.Basic',Map="/Engine/Maps/Templates/Basic")
\ No newline at end of file
diff --git a/Templates/TP_AEC_BlankBP/Config/TemplateDefs.ini b/Templates/TP_AEC_BlankBP/Config/TemplateDefs.ini
--- a/Templates/TP_AEC_BlankBP/Config/TemplateDefs.ini
+++ b/Templates/TP_AEC_BlankBP/Config/TemplateDefs.ini
@@ -81,7 +81,7 @@ HiddenSettings=Languages
 HiddenSettings=HardwareTarget
 HiddenSettings=GraphicsPreset 
 ;HiddenSettings=StarterContent 
-HiddenSettings=XR
+;HiddenSettings=XR
 
 ; The class types that this project uses.
 ; This is only for display in the template description, and does not affect project creation.
diff --git a/Templates/TP_ArchvisBP/Config/DefaultEngine.ini b/Templates/TP_ArchvisBP/Config/DefaultEngine.ini
--- a/Templates/TP_ArchvisBP/Config/DefaultEngine.ini
+++ b/Templates/TP_ArchvisBP/Config/DefaultEngine.ini
@@ -17,7 +17,7 @@ r.LightPropagationVolume=0
 ; quality settings
 r.ReflectionCaptureResolution=2048
 r.AllowStaticLighting=True
-r.HighResScreenshotDelay 8
+r.HighResScreenshotDelay=8
 r.DefaultBackBufferPixelFormat=4
 r.AllowGlobalClipPlane=False
 
