Archvis template:

fixing a few lose ends by renaming levels correctly and relinking levels to the default maps.

diff --git a/Templates/TP_AEC_BlankBP/Config/DefaultEngine.ini b/Templates/TP_AEC_BlankBP/Config/DefaultEngine.ini
--- a/Templates/TP_AEC_BlankBP/Config/DefaultEngine.ini
+++ b/Templates/TP_AEC_BlankBP/Config/DefaultEngine.ini
@@ -30,8 +30,8 @@ r.NormalMapsForStaticLighting=False
 
 [/Script/EngineSettings.GameMapsSettings]
 ;defines what level to be loaded by default in the editor and at playback time
-GameDefaultMap=/Game/Main.Main
-EditorStartupMap=/Game/Main.Main
+GameDefaultMap=/Game/Exterior.Exterior
+EditorStartupMap=/Game/Exterior.Exterior
 
 
 [/Script/UnrealEd.UnrealEdEngine]
