ArchVis Blank Template

Fixing default map not being loaded by the project on startup.

diff --git a/Templates/TP_AEC_BlankBP/Config/DefaultEngine.ini b/Templates/TP_AEC_BlankBP/Config/DefaultEngine.ini
--- a/Templates/TP_AEC_BlankBP/Config/DefaultEngine.ini
+++ b/Templates/TP_AEC_BlankBP/Config/DefaultEngine.ini
@@ -30,8 +30,8 @@ r.NormalMapsForStaticLighting=False
 
 [/Script/EngineSettings.GameMapsSettings]
 ;defines what level to be loaded by default in the editor and at playback time
-GameDefaultMap=/Game/Exterior.Exterior
-EditorStartupMap=/Game/Exterior.Exterior
+GameDefaultMap=/Game/Main.Main
+EditorStartupMap=/Game/Main.Main
 
 
 [/Script/UnrealEd.UnrealEdEngine]
