AEC Blank template:

Categorized in both AEC and MFG
Updated images and icons
Removed boxes in the level
Fixed asset preview for better color fidelity

diff --git a/Templates/TP_AEC_BlankBP/Config/TemplateDefs.ini b/Templates/TP_AEC_BlankBP/Config/TemplateDefs.ini
--- a/Templates/TP_AEC_BlankBP/Config/TemplateDefs.ini
+++ b/Templates/TP_AEC_BlankBP/Config/TemplateDefs.ini
@@ -6,14 +6,15 @@
 
 ; The categories that this template should appear in. The values in this list refer to the Key value given in TemplateCategories.ini.
 Categories=AEC
+Categories=MFG
 
 ; The localized display names for this template. Each entry specifies a Language code and Unicode text to display in the New Project dialog.
 LocalizedDisplayNames=(Language="en", Text="Blank")
 
 ; The localized descriptions for this template, which is displayed when the template is selected.
-LocalizedDescriptions=(Language="en", Text="Blank template with minimal functionality.")
+LocalizedDescriptions=(Language="en", Text="Blank template with minimal functionality and content.\n\nThis template enables you to use Datasmith to import models from various design applications and easily light them.")
 ; The sort key to use for the template to bypass the default alphabetical sorting.
-SortKey="_1"
+SortKey="0"
 
 ; A list of folders to avoid copying when creating a project from this template. 
 ; These should contain files only of concern to the template itself and not desired in the user's project.
