Blank AEC Template

Adding Fly and Walk pawn with proper input mappings
Adding Bookmarks that is used by the Game pawn
Default level updated to contain basic Pwn and Bookmarks

diff --git a/Templates/TP_AEC_BlankBP/Config/DefaultEngine.ini b/Templates/TP_AEC_BlankBP/Config/DefaultEngine.ini
--- a/Templates/TP_AEC_BlankBP/Config/DefaultEngine.ini
+++ b/Templates/TP_AEC_BlankBP/Config/DefaultEngine.ini
@@ -14,6 +14,7 @@ r.ClearCoatNormal=True
 r.AllowStaticLighting=True
 r.NormalMapsForStaticLighting=False
 r.DefaultBackBufferPixelFormat=4
+r.HighResScreenshotDelay 8
 
 [/Script/EngineSettings.GameMapsSettings]
 GameDefaultMap=/Game/Main.Main
diff --git a/Templates/TP_AEC_BlankBP/Config/DefaultInput.ini b/Templates/TP_AEC_BlankBP/Config/DefaultInput.ini
new file mode 100644
--- /dev/null
+++ b/Templates/TP_AEC_BlankBP/Config/DefaultInput.ini
@@ -0,0 +1,76 @@
+[/Script/Engine.InputSettings]
++AxisConfig=(AxisKeyName="MotionController_Right_Thumbstick_Y",AxisProperties=(DeadZone=0.000000,Sensitivity=1.000000,Exponent=1.000000,bInvert=False))
++AxisConfig=(AxisKeyName="MotionController_Right_TriggerAxis",AxisProperties=(DeadZone=0.000000,Sensitivity=1.000000,Exponent=1.000000,bInvert=False))
++AxisConfig=(AxisKeyName="MotionController_Right_Grip1Axis",AxisProperties=(DeadZone=0.000000,Sensitivity=1.000000,Exponent=1.000000,bInvert=False))
++AxisConfig=(AxisKeyName="MotionController_Right_Grip2Axis",AxisProperties=(DeadZone=0.000000,Sensitivity=1.000000,Exponent=1.000000,bInvert=False))
++AxisConfig=(AxisKeyName="Gamepad_Special_Left_X",AxisProperties=(DeadZone=0.000000,Sensitivity=1.000000,Exponent=1.000000,bInvert=False))
++AxisConfig=(AxisKeyName="Gamepad_Special_Left_Y",AxisProperties=(DeadZone=0.000000,Sensitivity=1.000000,Exponent=1.000000,bInvert=False))
++AxisConfig=(AxisKeyName="Gamepad_LeftX",AxisProperties=(DeadZone=0.250000,Sensitivity=1.000000,Exponent=1.000000,bInvert=False))
++AxisConfig=(AxisKeyName="Gamepad_LeftY",AxisProperties=(DeadZone=0.250000,Sensitivity=1.000000,Exponent=1.000000,bInvert=False))
++AxisConfig=(AxisKeyName="Gamepad_RightX",AxisProperties=(DeadZone=0.250000,Sensitivity=1.000000,Exponent=1.000000,bInvert=False))
++AxisConfig=(AxisKeyName="Gamepad_RightY",AxisProperties=(DeadZone=0.250000,Sensitivity=1.000000,Exponent=1.000000,bInvert=False))
++AxisConfig=(AxisKeyName="MouseX",AxisProperties=(DeadZone=0.000000,Sensitivity=0.070000,Exponent=1.000000,bInvert=False))
++AxisConfig=(AxisKeyName="MouseY",AxisProperties=(DeadZone=0.000000,Sensitivity=0.070000,Exponent=1.000000,bInvert=False))
++AxisConfig=(AxisKeyName="MouseWheelAxis",AxisProperties=(DeadZone=0.000000,Sensitivity=1.000000,Exponent=1.000000,bInvert=False))
++AxisConfig=(AxisKeyName="Gamepad_LeftTriggerAxis",AxisProperties=(DeadZone=0.000000,Sensitivity=1.000000,Exponent=1.000000,bInvert=False))
++AxisConfig=(AxisKeyName="Gamepad_RightTriggerAxis",AxisProperties=(DeadZone=0.000000,Sensitivity=1.000000,Exponent=1.000000,bInvert=False))
++AxisConfig=(AxisKeyName="MotionController_Left_Thumbstick_X",AxisProperties=(DeadZone=0.000000,Sensitivity=1.000000,Exponent=1.000000,bInvert=False))
++AxisConfig=(AxisKeyName="MotionController_Left_Thumbstick_Y",AxisProperties=(DeadZone=0.000000,Sensitivity=1.000000,Exponent=1.000000,bInvert=False))
++AxisConfig=(AxisKeyName="MotionController_Left_TriggerAxis",AxisProperties=(DeadZone=0.000000,Sensitivity=1.000000,Exponent=1.000000,bInvert=False))
++AxisConfig=(AxisKeyName="MotionController_Left_Grip1Axis",AxisProperties=(DeadZone=0.000000,Sensitivity=1.000000,Exponent=1.000000,bInvert=False))
++AxisConfig=(AxisKeyName="MotionController_Left_Grip2Axis",AxisProperties=(DeadZone=0.000000,Sensitivity=1.000000,Exponent=1.000000,bInvert=False))
++AxisConfig=(AxisKeyName="MotionController_Right_Thumbstick_X",AxisProperties=(DeadZone=0.000000,Sensitivity=1.000000,Exponent=1.000000,bInvert=False))
+bAltEnterTogglesFullscreen=True
+bF11TogglesFullscreen=True
+bUseMouseForTouch=False
+bEnableMouseSmoothing=True
+bEnableFOVScaling=True
+bCaptureMouseOnLaunch=True
+bDefaultViewportMouseLock=False
+bAlwaysShowTouchInterface=False
+bShowConsoleOnFourFingerTap=True
+bEnableGestureRecognizer=False
+bUseAutocorrect=False
+DefaultViewportMouseCaptureMode=CapturePermanently_IncludingInitialMouseDown
+DefaultViewportMouseLockMode=LockOnCapture
+FOVScale=0.011110
+DoubleClickTime=0.200000
++ActionMappings=(ActionName="NextBookmark",bShift=False,bCtrl=False,bAlt=False,bCmd=False,Key=Right)
++ActionMappings=(ActionName="PreviousBookmark",bShift=False,bCtrl=False,bAlt=False,bCmd=False,Key=Left)
++ActionMappings=(ActionName="Run",bShift=False,bCtrl=False,bAlt=False,bCmd=False,Key=LeftShift)
++ActionMappings=(ActionName="Run",bShift=False,bCtrl=False,bAlt=False,bCmd=False,Key=RightShift)
++AxisMappings=(AxisName="FlyUp",Scale=1.000000,Key=E)
++AxisMappings=(AxisName="FlyUp",Scale=-1.000000,Key=Q)
++AxisMappings=(AxisName="LookRight",Scale=1.000000,Key=MouseX)
++AxisMappings=(AxisName="LookUp",Scale=1.000000,Key=MouseY)
++AxisMappings=(AxisName="MoveRight",Scale=-1.000000,Key=A)
++AxisMappings=(AxisName="MoveRight",Scale=1.000000,Key=D)
++AxisMappings=(AxisName="MoveRight",Scale=-1.000000,Key=Gamepad_LeftStick_Left)
++AxisMappings=(AxisName="MoveRight",Scale=1.000000,Key=Gamepad_LeftStick_Right)
++AxisMappings=(AxisName="MoveRight",Scale=1.000000,Key=Gamepad_LeftX)
++AxisMappings=(AxisName="MoveRight",Scale=-1.000000,Key=Left)
++AxisMappings=(AxisName="MoveRight",Scale=1.000000,Key=Right)
++AxisMappings=(AxisName="MoveUp",Scale=-1.000000,Key=Down)
++AxisMappings=(AxisName="MoveUp",Scale=-1.000000,Key=Gamepad_LeftStick_Down)
++AxisMappings=(AxisName="MoveUp",Scale=1.000000,Key=Gamepad_LeftStick_Up)
++AxisMappings=(AxisName="MoveUp",Scale=1.000000,Key=Gamepad_LeftY)
++AxisMappings=(AxisName="MoveUp",Scale=-1.000000,Key=S)
++AxisMappings=(AxisName="MoveUp",Scale=1.000000,Key=Up)
++AxisMappings=(AxisName="MoveUp",Scale=1.000000,Key=W)
++AxisMappings=(AxisName="Thrust",Scale=-1.000000,Key=Gamepad_LeftTriggerAxis)
++AxisMappings=(AxisName="Thrust",Scale=1.000000,Key=Gamepad_RightTriggerAxis)
++AxisMappings=(AxisName="Thrust",Scale=1.000000,Key=Gamepad_RightY)
++AxisMappings=(AxisName="Thrust",Scale=-1.000000,Key=LeftControl)
++AxisMappings=(AxisName="Thrust",Scale=1.000000,Key=LeftShift)
++AxisMappings=(AxisName="Thrust",Scale=-1.000000,Key=RightControl)
++AxisMappings=(AxisName="Thrust",Scale=1.000000,Key=RightShift)
++AxisMappings=(AxisName="Thrust",Scale=1.000000,Key=SpaceBar)
++AxisMappings=(AxisName="WalkForward",Scale=-1.000000,Key=S)
++AxisMappings=(AxisName="WalkForward",Scale=1.000000,Key=W)
++AxisMappings=(AxisName="WalkRight",Scale=-1.000000,Key=A)
++AxisMappings=(AxisName="WalkRight",Scale=1.000000,Key=D)
+DefaultTouchInterface=/Engine/MobileResources/HUD/DefaultVirtualJoysticks.DefaultVirtualJoysticks
+ConsoleKey=None
+ConsoleKeys=Tilde
+
+
