AEC Blank template improvements

Fixed the pawn, removed bookmarks, tested walk, fly, with desktop and game controller
Updated default lightmass settings suggested by Marien
Continued adding ini settings based on previous discussions.

diff --git a/Templates/TP_AEC_BlankBP/Config/DefaultEngine.ini b/Templates/TP_AEC_BlankBP/Config/DefaultEngine.ini
--- a/Templates/TP_AEC_BlankBP/Config/DefaultEngine.ini
+++ b/Templates/TP_AEC_BlankBP/Config/DefaultEngine.ini
@@ -1,27 +1,39 @@
 [/Script/HardwareTargeting.HardwareTargetingSettings]
+; quality settings
 TargetedHardwareClass=Desktop
 AppliedTargetedHardwareClass=Desktop
 DefaultGraphicsPerformance=Maximum
 AppliedDefaultGraphicsPerformance=Maximum
 
 [/Script/Engine.RendererSettings]
+; exposure control and tone mapping settings
+r.Tonemapper.Sharpen=2
+r.TemporalAA.Upsampling=True
 r.DefaultFeature.AutoExposure.ExtendDefaultLuminanceRange=True
-r.SupportSkyAtmosphereAffectsHeightFog=True
-r.ReflectionCaptureResolution=1024
+
+;dynamic gi settings
+r.LightPropagationVolume=1
+
+; quality settings
+r.ReflectionCaptureResolution=2048
+r.AllowStaticLighting=True
+r.HighResScreenshotDelay 8
+r.DefaultBackBufferPixelFormat=4
 r.AllowGlobalClipPlane=False
+
+; shader settings
+r.SupportSkyAtmosphereAffectsHeightFog=True
 r.DefaultFeature.LensFlare=True
 r.ClearCoatNormal=True
-r.AllowStaticLighting=True
 r.NormalMapsForStaticLighting=False
-r.DefaultBackBufferPixelFormat=4
-r.HighResScreenshotDelay 8
+
 
 [/Script/EngineSettings.GameMapsSettings]
+;defines what level to be loaded by default in the editor and at playback time
 GameDefaultMap=/Game/Main.Main
 EditorStartupMap=/Game/Main.Main
 
 
-
 [/Script/UnrealEd.UnrealEdEngine]
 ; Remove Engine Template maps
 -TemplateMapInfos=(ThumbnailTexture=Texture2D'/Engine/Maps/Templates/Thumbnails/Default.Default',Map="/Engine/Maps/Templates/Template_Default")
diff --git a/Templates/TP_AEC_BlankBP/Config/DefaultInput.ini b/Templates/TP_AEC_BlankBP/Config/DefaultInput.ini
--- a/Templates/TP_AEC_BlankBP/Config/DefaultInput.ini
+++ b/Templates/TP_AEC_BlankBP/Config/DefaultInput.ini
@@ -1,76 +0,0 @@
-[/Script/Engine.InputSettings]
-+AxisConfig=(AxisKeyName="MotionController_Right_Thumbstick_Y",AxisProperties=(DeadZone=0.000000,Sensitivity=1.000000,Exponent=1.000000,bInvert=False))
-+AxisConfig=(AxisKeyName="MotionController_Right_TriggerAxis",AxisProperties=(DeadZone=0.000000,Sensitivity=1.000000,Exponent=1.000000,bInvert=False))
-+AxisConfig=(AxisKeyName="MotionController_Right_Grip1Axis",AxisProperties=(DeadZone=0.000000,Sensitivity=1.000000,Exponent=1.000000,bInvert=False))
-+AxisConfig=(AxisKeyName="MotionController_Right_Grip2Axis",AxisProperties=(DeadZone=0.000000,Sensitivity=1.000000,Exponent=1.000000,bInvert=False))
-+AxisConfig=(AxisKeyName="Gamepad_Special_Left_X",AxisProperties=(DeadZone=0.000000,Sensitivity=1.000000,Exponent=1.000000,bInvert=False))
-+AxisConfig=(AxisKeyName="Gamepad_Special_Left_Y",AxisProperties=(DeadZone=0.000000,Sensitivity=1.000000,Exponent=1.000000,bInvert=False))
-+AxisConfig=(AxisKeyName="Gamepad_LeftX",AxisProperties=(DeadZone=0.250000,Sensitivity=1.000000,Exponent=1.000000,bInvert=False))
-+AxisConfig=(AxisKeyName="Gamepad_LeftY",AxisProperties=(DeadZone=0.250000,Sensitivity=1.000000,Exponent=1.000000,bInvert=False))
-+AxisConfig=(AxisKeyName="Gamepad_RightX",AxisProperties=(DeadZone=0.250000,Sensitivity=1.000000,Exponent=1.000000,bInvert=False))
-+AxisConfig=(AxisKeyName="Gamepad_RightY",AxisProperties=(DeadZone=0.250000,Sensitivity=1.000000,Exponent=1.000000,bInvert=False))
-+AxisConfig=(AxisKeyName="MouseX",AxisProperties=(DeadZone=0.000000,Sensitivity=0.070000,Exponent=1.000000,bInvert=False))
-+AxisConfig=(AxisKeyName="MouseY",AxisProperties=(DeadZone=0.000000,Sensitivity=0.070000,Exponent=1.000000,bInvert=False))
-+AxisConfig=(AxisKeyName="MouseWheelAxis",AxisProperties=(DeadZone=0.000000,Sensitivity=1.000000,Exponent=1.000000,bInvert=False))
-+AxisConfig=(AxisKeyName="Gamepad_LeftTriggerAxis",AxisProperties=(DeadZone=0.000000,Sensitivity=1.000000,Exponent=1.000000,bInvert=False))
-+AxisConfig=(AxisKeyName="Gamepad_RightTriggerAxis",AxisProperties=(DeadZone=0.000000,Sensitivity=1.000000,Exponent=1.000000,bInvert=False))
-+AxisConfig=(AxisKeyName="MotionController_Left_Thumbstick_X",AxisProperties=(DeadZone=0.000000,Sensitivity=1.000000,Exponent=1.000000,bInvert=False))
-+AxisConfig=(AxisKeyName="MotionController_Left_Thumbstick_Y",AxisProperties=(DeadZone=0.000000,Sensitivity=1.000000,Exponent=1.000000,bInvert=False))
-+AxisConfig=(AxisKeyName="MotionController_Left_TriggerAxis",AxisProperties=(DeadZone=0.000000,Sensitivity=1.000000,Exponent=1.000000,bInvert=False))
-+AxisConfig=(AxisKeyName="MotionController_Left_Grip1Axis",AxisProperties=(DeadZone=0.000000,Sensitivity=1.000000,Exponent=1.000000,bInvert=False))
-+AxisConfig=(AxisKeyName="MotionController_Left_Grip2Axis",AxisProperties=(DeadZone=0.000000,Sensitivity=1.000000,Exponent=1.000000,bInvert=False))
-+AxisConfig=(AxisKeyName="MotionController_Right_Thumbstick_X",AxisProperties=(DeadZone=0.000000,Sensitivity=1.000000,Exponent=1.000000,bInvert=False))
-bAltEnterTogglesFullscreen=True
-bF11TogglesFullscreen=True
-bUseMouseForTouch=False
-bEnableMouseSmoothing=True
-bEnableFOVScaling=True
-bCaptureMouseOnLaunch=True
-bDefaultViewportMouseLock=False
-bAlwaysShowTouchInterface=False
-bShowConsoleOnFourFingerTap=True
-bEnableGestureRecognizer=False
-bUseAutocorrect=False
-DefaultViewportMouseCaptureMode=CapturePermanently_IncludingInitialMouseDown
-DefaultViewportMouseLockMode=LockOnCapture
-FOVScale=0.011110
-DoubleClickTime=0.200000
-+ActionMappings=(ActionName="NextBookmark",bShift=False,bCtrl=False,bAlt=False,bCmd=False,Key=Right)
-+ActionMappings=(ActionName="PreviousBookmark",bShift=False,bCtrl=False,bAlt=False,bCmd=False,Key=Left)
-+ActionMappings=(ActionName="Run",bShift=False,bCtrl=False,bAlt=False,bCmd=False,Key=LeftShift)
-+ActionMappings=(ActionName="Run",bShift=False,bCtrl=False,bAlt=False,bCmd=False,Key=RightShift)
-+AxisMappings=(AxisName="FlyUp",Scale=1.000000,Key=E)
-+AxisMappings=(AxisName="FlyUp",Scale=-1.000000,Key=Q)
-+AxisMappings=(AxisName="LookRight",Scale=1.000000,Key=MouseX)
-+AxisMappings=(AxisName="LookUp",Scale=1.000000,Key=MouseY)
-+AxisMappings=(AxisName="MoveRight",Scale=-1.000000,Key=A)
-+AxisMappings=(AxisName="MoveRight",Scale=1.000000,Key=D)
-+AxisMappings=(AxisName="MoveRight",Scale=-1.000000,Key=Gamepad_LeftStick_Left)
-+AxisMappings=(AxisName="MoveRight",Scale=1.000000,Key=Gamepad_LeftStick_Right)
-+AxisMappings=(AxisName="MoveRight",Scale=1.000000,Key=Gamepad_LeftX)
-+AxisMappings=(AxisName="MoveRight",Scale=-1.000000,Key=Left)
-+AxisMappings=(AxisName="MoveRight",Scale=1.000000,Key=Right)
-+AxisMappings=(AxisName="MoveUp",Scale=-1.000000,Key=Down)
-+AxisMappings=(AxisName="MoveUp",Scale=-1.000000,Key=Gamepad_LeftStick_Down)
-+AxisMappings=(AxisName="MoveUp",Scale=1.000000,Key=Gamepad_LeftStick_Up)
-+AxisMappings=(AxisName="MoveUp",Scale=1.000000,Key=Gamepad_LeftY)
-+AxisMappings=(AxisName="MoveUp",Scale=-1.000000,Key=S)
-+AxisMappings=(AxisName="MoveUp",Scale=1.000000,Key=Up)
-+AxisMappings=(AxisName="MoveUp",Scale=1.000000,Key=W)
-+AxisMappings=(AxisName="Thrust",Scale=-1.000000,Key=Gamepad_LeftTriggerAxis)
-+AxisMappings=(AxisName="Thrust",Scale=1.000000,Key=Gamepad_RightTriggerAxis)
-+AxisMappings=(AxisName="Thrust",Scale=1.000000,Key=Gamepad_RightY)
-+AxisMappings=(AxisName="Thrust",Scale=-1.000000,Key=LeftControl)
-+AxisMappings=(AxisName="Thrust",Scale=1.000000,Key=LeftShift)
-+AxisMappings=(AxisName="Thrust",Scale=-1.000000,Key=RightControl)
-+AxisMappings=(AxisName="Thrust",Scale=1.000000,Key=RightShift)
-+AxisMappings=(AxisName="Thrust",Scale=1.000000,Key=SpaceBar)
-+AxisMappings=(AxisName="WalkForward",Scale=-1.000000,Key=S)
-+AxisMappings=(AxisName="WalkForward",Scale=1.000000,Key=W)
-+AxisMappings=(AxisName="WalkRight",Scale=-1.000000,Key=A)
-+AxisMappings=(AxisName="WalkRight",Scale=1.000000,Key=D)
-DefaultTouchInterface=/Engine/MobileResources/HUD/DefaultVirtualJoysticks.DefaultVirtualJoysticks
-ConsoleKey=None
-ConsoleKeys=Tilde
-
-
diff --git a/Templates/TP_AEC_BlankBP/Config/DefaultLightmass.ini b/Templates/TP_AEC_BlankBP/Config/DefaultLightmass.ini
new file mode 100644
--- /dev/null
+++ b/Templates/TP_AEC_BlankBP/Config/DefaultLightmass.ini
@@ -0,0 +1,22 @@
+; These are tweaked defaults for various lightmass solver and export settings
+; Artist oriented lightmass settings are in the editor UI
+; Documentation for all of these is in UnrealLightmass / Public / SceneExport.h
+; This ini is reloaded every time a lighting build begins, no need to restart
+
+; Warning: overwriting this file with an old version will cause the editor to crash.  This file must be in sync with the editor executable.
+; Instead, create a DefaultLightmass.ini in your project and override just the values you need, then the overrides will continue to work on version upgrades.
+; https://docs.unrealengine.com/latest/INT/Programming/Basics/ConfigurationFiles/
+;
+; For example, in your project's Config/DefaultLightmass.ini:
+; [DevOptions.PrecomputedDynamicObjectLighting]
+; SurfaceLightSampleSpacing=400
+
+[DevOptions.StaticLightingSettings]
+bCompressLightmaps=False
+NumIndirectLightingBounces=5
+NumSkyLightingBounces=5
+IndirectLightingSmoothness=0.9
+IndirectLightingQuality=2.0
+
+[DevOptions.StaticLightingSceneConstants]
+StaticLightingLevelScale=0.5
\ No newline at end of file
diff --git a/Templates/TP_AEC_BlankBP/Config/TemplateDefs.ini b/Templates/TP_AEC_BlankBP/Config/TemplateDefs.ini
--- a/Templates/TP_AEC_BlankBP/Config/TemplateDefs.ini
+++ b/Templates/TP_AEC_BlankBP/Config/TemplateDefs.ini
@@ -11,7 +11,7 @@ Categories=AEC
 LocalizedDisplayNames=(Language="en", Text="Blank")
 
 ; The localized descriptions for this template, which is displayed when the template is selected.
-LocalizedDescriptions=(Language="en", Text="This template is designed to create architectural walkthroughs with Datasmith. The player will be able to navigate with a fly and walk mode and bookmarks.")
+LocalizedDescriptions=(Language="en", Text="Blank AEC Template with minimal functionality.")
 ; The sort key to use for the template to bypass the default alphabetical sorting.
 SortKey="_1"
 
@@ -81,7 +81,7 @@ HiddenSettings=Languages
 HiddenSettings=HardwareTarget
 HiddenSettings=GraphicsPreset 
 ;HiddenSettings=StarterContent 
-HiddenSettings=VR
+HiddenSettings=XR
 
 ; The class types that this project uses.
 ; This is only for display in the template description, and does not affect project creation.
