AEC Blank Template
Minor tweaks
Enabling the desired plugins

diff --git a/Templates/TP_AEC_BlankBP/Config/DefaultEngine.ini b/Templates/TP_AEC_BlankBP/Config/DefaultEngine.ini
--- a/Templates/TP_AEC_BlankBP/Config/DefaultEngine.ini
+++ b/Templates/TP_AEC_BlankBP/Config/DefaultEngine.ini
@@ -12,7 +12,7 @@ r.TemporalAA.Upsampling=True
 r.DefaultFeature.AutoExposure.ExtendDefaultLuminanceRange=True
 
 ;dynamic gi settings
-r.LightPropagationVolume=1
+r.LightPropagationVolume=0
 
 ; quality settings
 r.ReflectionCaptureResolution=2048
@@ -24,7 +24,7 @@ r.AllowGlobalClipPlane=False
 ; shader settings
 r.SupportSkyAtmosphereAffectsHeightFog=True
 r.DefaultFeature.LensFlare=True
-r.ClearCoatNormal=True
+r.ClearCoatNormal=False
 r.NormalMapsForStaticLighting=False
 
 
diff --git a/Templates/TP_AEC_BlankBP/TP_AEC_BlankBP.uproject b/Templates/TP_AEC_BlankBP/TP_AEC_BlankBP.uproject
--- a/Templates/TP_AEC_BlankBP/TP_AEC_BlankBP.uproject
+++ b/Templates/TP_AEC_BlankBP/TP_AEC_BlankBP.uproject
@@ -3,7 +3,6 @@
 	"EngineAssociation": "",
 	"Category": "",
 	"Description": "",
-	"Modules": [],
 	"Enterprise": true,
 	"Plugins": [
 		{
@@ -17,6 +16,18 @@
 		{
 			"Name": "DatasmithImporter",
 			"Enabled": true
+		},
+		{
+			"Name": "VariantManager",
+			"Enabled": True
+		},
+		{
+			"Name": "PythonScriptPlugin",
+			"Enabled": true
+		},
+				{
+			"Name": "HDRIBackdrop",
+			"Enabled": true
 		}
 	]
 }
\ No newline at end of file
