AEC template  Minor tweak in description

diff --git a/Templates/TP_AEC_BlankBP/Config/TemplateDefs.ini b/Templates/TP_AEC_BlankBP/Config/TemplateDefs.ini
--- a/Templates/TP_AEC_BlankBP/Config/TemplateDefs.ini
+++ b/Templates/TP_AEC_BlankBP/Config/TemplateDefs.ini
@@ -11,7 +11,7 @@ Categories=AEC
 LocalizedDisplayNames=(Language="en", Text="Blank")
 
 ; The localized descriptions for this template, which is displayed when the template is selected.
-LocalizedDescriptions=(Language="en", Text="Blank AEC Template with minimal functionality.")
+LocalizedDescriptions=(Language="en", Text="Blank template with minimal functionality.")
 ; The sort key to use for the template to bypass the default alphabetical sorting.
 SortKey="_1"
 
