Rename/move file(s)
Renaming a misnamed bitmap

diff --git a/Templates/TP_AEC_BlankBP/Config/DefaultEngine.ini b/Templates/TP_AEC_BlankBP/Config/DefaultEngine.ini
--- a/Templates/TP_AEC_BlankBP/Config/DefaultEngine.ini
+++ b/Templates/TP_AEC_BlankBP/Config/DefaultEngine.ini
@@ -41,4 +41,4 @@ EditorStartupMap=/Game/Main.Main
 
 ; Add Predefined Template Maps that are suitable for this template
 ;+TemplateMapInfos=(ThumbnailTexture=Texture2D'/Game/Developers/Maps/Thumbnails/AssetName.AssetName',Map="/Game/Maps/Templates/MapName")
-+TemplateMapInfos=(ThumbnailTexture=Texture2D'/Game/Developers/LevelTemplates/Thumbnails/T_Default.T_Default',Map="/Game/Developers/LevelTemplates/Default")
\ No newline at end of file
++TemplateMapInfos=(ThumbnailTexture=Texture2D'/Game/Developers/LevelTemplates/Thumbnails/Default.Default',Map="/Game/Developers/LevelTemplates/Default")
\ No newline at end of file
