adding Chinese translation to the file #jira ueloc-4711 #fyi jamie.dale, michelle.zhao, alexis.zhang, yali.he

diff --git a/Templates/TP_AEC_BlankBP/Config/TemplateDefs.ini b/Templates/TP_AEC_BlankBP/Config/TemplateDefs.ini
--- a/Templates/TP_AEC_BlankBP/Config/TemplateDefs.ini
+++ b/Templates/TP_AEC_BlankBP/Config/TemplateDefs.ini
@@ -11,12 +11,17 @@ Categories=MFG
 ; The localized display names for this template. Each entry specifies a Language code and Unicode text to display in the New Project dialog.
 LocalizedDisplayNames=(Language="en", Text="Blank")
 
+
 ; The localized descriptions for this template, which is displayed when the template is selected.
 LocalizedDescriptions=(Language="en", Text="Blank template with minimal functionality and content.\n\nThis template enables you to use Datasmith to import models from various design applications and easily light them.")
 
 LocalizedDisplayNames=(Language="ko", Text="빈 프로젝트")
 LocalizedDescriptions=(Language="ko", Text="최소한의 기능과 콘텐츠를 포함한 빈 프로젝트 템플릿\n\n이 템플릿은 다양한 디자인 어플리케이션에서 데이터스미스를 이용한 모델 임포트와 시각화가 가능합니다.")
 
+LocalizedDisplayNames=(Language="zh-Hans", Text="空白")
+LocalizedDescriptions=(Language="zh-Hans", Text="具备基本功能和内容的空白模板。\n\n该模板允许使用Datasmith从多种设计软件中导入模型并轻松为其设置光照。")
+
+
 ; The sort key to use for the template to bypass the default alphabetical sorting.
 SortKey="0"
 
