Adding Japanese translation #jira ueloc-4712 #fyi michelle.zhao, yali.he

diff --git a/Templates/TP_AEC_BlankBP/Config/TemplateDefs.ini b/Templates/TP_AEC_BlankBP/Config/TemplateDefs.ini
--- a/Templates/TP_AEC_BlankBP/Config/TemplateDefs.ini
+++ b/Templates/TP_AEC_BlankBP/Config/TemplateDefs.ini
@@ -21,6 +21,8 @@ LocalizedDescriptions=(Language="ko", Text="최소한의 기능과 콘텐츠를
 LocalizedDisplayNames=(Language="zh-Hans", Text="空白")
 LocalizedDescriptions=(Language="zh-Hans", Text="具备基本功能和内容的空白模板。\n\n该模板允许使用Datasmith从多种设计软件中导入模型并轻松为其设置光照。")
 
+LocalizedDescriptions=(Language="ja", Text="最小限の機能とコンテンツが備わった Blank テンプレートです。\n\nこのテンプレートを利用すると、Datasmith を使ってモデルをさまざまなデザイン アプリケーションからインポートできるようになるとともに、それらを簡単にライティングできるようになります。")
+
 
 ; The sort key to use for the template to bypass the default alphabetical sorting.
 SortKey="0"
