Updated default plugins and cfg on TP_AEC_BlankBP

diff --git a/Templates/TP_AEC_BlankBP/Config/DefaultEngine.ini b/Templates/TP_AEC_BlankBP/Config/DefaultEngine.ini
--- a/Templates/TP_AEC_BlankBP/Config/DefaultEngine.ini
+++ b/Templates/TP_AEC_BlankBP/Config/DefaultEngine.ini
@@ -25,7 +25,15 @@ r.SupportSkyAtmosphereAffectsHeightFog=True
 r.DefaultFeature.LensFlare=True
 r.ClearCoatNormal=False
 r.NormalMapsForStaticLighting=False
-
+r.VirtualTextures=True
+r.VirtualTexturedLightmaps=True
+r.ReflectionMethod=1
+r.GenerateMeshDistanceFields=True
+r.DynamicGlobalIlluminationMethod=1
+r.SkinCache.CompileShaders=True
+r.RayTracing=True
+r.PathTracing=True
+r.RayTracing.Shadows=True
 
 [/Script/EngineSettings.GameMapsSettings]
 ;defines what level to be loaded by default in the editor and at playback time
@@ -37,4 +45,4 @@ EditorStartupMap=/Game/Main.Main
 ; Remove Engine Template maps
 -TemplateMapInfos=(ThumbnailTexture=Texture2D'/Engine/Maps/Templates/Thumbnails/Default.Default',Map="/Engine/Maps/Templates/Template_Default")
 +TemplateMapInfos=(ThumbnailTexture=Texture2D'/Engine/Maps/Templates/Thumbnails/TimeOfDay.TimeOfDay',Map="/Engine/Maps/Templates/TimeOfDay_Default")
--TemplateMapInfos=(ThumbnailTexture=Texture2D'/Engine/Maps/Templates/Thumbnails/VR-Basic.VR-Basic',Map="/Engine/Maps/Templates/VR-Basic")
\ No newline at end of file
+-TemplateMapInfos=(ThumbnailTexture=Texture2D'/Engine/Maps/Templates/Thumbnails/VR-Basic.VR-Basic',Map="/Engine/Maps/Templates/VR-Basic")
diff --git a/Templates/TP_AEC_BlankBP/TP_AEC_BlankBP.uproject b/Templates/TP_AEC_BlankBP/TP_AEC_BlankBP.uproject
--- a/Templates/TP_AEC_BlankBP/TP_AEC_BlankBP.uproject
+++ b/Templates/TP_AEC_BlankBP/TP_AEC_BlankBP.uproject
@@ -73,6 +73,18 @@
 		{
 			"Name": "MovieRenderPipeline",
 			"Enabled": true
+		},
+		{
+			"Name": "MoviePipelineMaskRenderPass",
+			"Enabled": true
+		},
+		{
+			"Name": "MoviePipelineDeadlineIntegration",
+			"Enabled": true
+		},
+		{
+			"Name": "GPULightmass",
+			"Enabled": true
 		}
 	]
 }
\ No newline at end of file
