Updated New Level entries for TP_AEC_Blank and TP_AEC_ArchvisBP #preflight none

diff --git a/Templates/TP_AEC_ArchvisBP/Config/DefaultEngine.ini b/Templates/TP_AEC_ArchvisBP/Config/DefaultEngine.ini
--- a/Templates/TP_AEC_ArchvisBP/Config/DefaultEngine.ini
+++ b/Templates/TP_AEC_ArchvisBP/Config/DefaultEngine.ini
@@ -43,8 +43,10 @@ r.RayTracing.Geometry.InstancedStaticMeshes.Culling=0
 
 [/Script/EngineSettings.GameMapsSettings]
 EditorStartupMap=/Game/ArchvisProject/Levels/Exterior.Exterior
++EditorTemplateMapOverrides=(Thumbnail=None,Map="/Game/ArchvisDefault/Levels/Archvis_Empty.Archvis_Empty",DisplayName=NSLOCTEXT("[/Script/EngineSettings]", "BDD1166D4B012519C508EE8B74D639B8", "Archvis Empty"))
 +EditorTemplateMapOverrides=(Thumbnail=None,Map="/Game/ArchvisProject/Levels/Exterior.Exterior",DisplayName=NSLOCTEXT("[/Script/EngineSettings]", "73A96829429462B471D06689EC07B2FB", "Exterior"))
-+EditorTemplateMapOverrides=(Thumbnail=None,Map="/Game/ArchvisProject/Levels/Interior_Lumen.Interior_Lumen",DisplayName=NSLOCTEXT("[/Script/EngineSettings]", "BDD1166D4B012519C508EE8B74D639B8", "Interior"))
++EditorTemplateMapOverrides=(Thumbnail=None,Map="/Game/ArchvisProject/Levels/Interior.Interior",DisplayName=NSLOCTEXT("[/Script/EngineSettings]", "2C6771E54254F7715F052A8A5B18AB8C", "Interior"))
++EditorTemplateMapOverrides=(Thumbnail=None,Map="/Game/ArchvisProject/Levels/Interior_Lightmass.Interior_Lightmass",DisplayName=NSLOCTEXT("[/Script/EngineSettings]", "7D660B88419A2FE50C151FA4408A8D80", "Interior Lightmass"))
 LocalMapOptions=
 TransitionMap=None
 bUseSplitscreen=True
diff --git a/Templates/TP_AEC_BlankBP/Config/DefaultEngine.ini b/Templates/TP_AEC_BlankBP/Config/DefaultEngine.ini
--- a/Templates/TP_AEC_BlankBP/Config/DefaultEngine.ini
+++ b/Templates/TP_AEC_BlankBP/Config/DefaultEngine.ini
@@ -41,16 +41,28 @@ r.Lumen.HardwareRayTracing=True
 r.Lumen.TraceMeshSDF=1
 
 [/Script/EngineSettings.GameMapsSettings]
-;defines what level to be loaded by default in the editor and at playback time
-GameDefaultMap=/Game/Main.Main
-EditorStartupMap=/Game/Main.Main
-
+EditorStartupMap=/Game/ArchvisDefault/Levels/Archvis_Empty.Archvis_Empty
++EditorTemplateMapOverrides=(Thumbnail=None,Map="/Game/ArchvisDefault/Levels/Archvis_Empty.Archvis_Empty",DisplayName=NSLOCTEXT("[/Script/EngineSettings]", "E1C1BC224C1F413942879D8305132FF6", "Archvis Empty"))
+LocalMapOptions=
+TransitionMap=None
+bUseSplitscreen=True
+TwoPlayerSplitscreenLayout=Horizontal
+ThreePlayerSplitscreenLayout=FavorTop
+FourPlayerSplitscreenLayout=Grid
+bOffsetPlayerGamepadIds=False
+GameInstanceClass=/Script/Engine.GameInstance
+GameDefaultMap=/Game/ArchvisDefault/Levels/Archvis_Empty.Archvis_Empty
+ServerDefaultMap=/Engine/Maps/Entry.Entry
+GlobalDefaultGameMode=/Script/Engine.GameModeBase
+GlobalDefaultServerGameMode=None
 
 [/Script/UnrealEd.UnrealEdEngine]
 ; Remove Engine Template maps
--TemplateMapInfos=(ThumbnailTexture=Texture2D'/Engine/Maps/Templates/Thumbnails/Default.Default',Map="/Engine/Maps/Templates/Template_Default")
-+TemplateMapInfos=(ThumbnailTexture=Texture2D'/Engine/Maps/Templates/Thumbnails/TimeOfDay.TimeOfDay',Map="/Engine/Maps/Templates/TimeOfDay_Default")
--TemplateMapInfos=(ThumbnailTexture=Texture2D'/Engine/Maps/Templates/Thumbnails/VR-Basic.VR-Basic',Map="/Engine/Maps/Templates/VR-Basic")
+-TemplateMapInfos=(Thumbnail="/Engine/Maps/Templates/Thumbnails/OpenWorld",Map="/Engine/Maps/Templates/OpenWorld")
+-TemplateMapInfos=(Thumbnail="/Engine/Maps/Templates/Thumbnails/Basic",Map="/Engine/Maps/Templates/Template_Default")
+-TemplateMapInfos=(Thumbnail="/Engine/Maps/Templates/Thumbnails/TimeOfDay",Map="/Engine/Maps/Templates/TimeOfDay_Default")
+-TemplateMapInfos=(Thumbnail="/Engine/Maps/Templates/Thumbnails/VR-Basic",Map="/Engine/Maps/Templates/VR-Basic")
++TemplateMapInfos=(Thumbnail="/Engine/Maps/Templates/Thumbnails/Default",Map="/Game/Content/Main")
 
 [/Script/WindowsTargetPlatform.WindowsTargetSettings]
 DefaultGraphicsRHI=DefaultGraphicsRHI_DX12
